[AI][WorldConditions] Add WorldConditionsToolset plugin for MCP inspection

Exposes WorldCondition queries to AI agents via the MCP ToolsetRegistry. Previously WorldCondition properties serialized as empty strings; now they produce full JSON with schema, conditions, and human-readable descriptions. Full round-trip supported -- MCP callers can pass WorldConditionQueryDefinition between tools or create from scratch.

Changes:
- New plugin: Engine/Plugins/Experimental/WorldConditionsToolset (editor-only)
- C++ toolset (UWorldConditionTools) with GetQueryDescription and GetConditionDescription
- Custom JSON converter (FWorldConditionQueryConverter) for FWorldConditionQueryDefinition round-trip: delegates to ToolsetJson pipeline, adds description field, calls Initialize() after import
- BlueprintType on FWorldConditionQueryDefinition for Blueprint/Python access

#Synced-CL 51880525
#Synced-CL 51880525

diff --git a/Engine/Plugins/Experimental/WorldConditionsToolset/Source/WorldConditionsToolset/Private/WorldConditionQueryConverter.cpp b/Engine/Plugins/Experimental/WorldConditionsToolset/Source/WorldConditionsToolset/Private/WorldConditionQueryConverter.cpp
new file mode 100644
--- /dev/null
+++ b/Engine/Plugins/Experimental/WorldConditionsToolset/Source/WorldConditionsToolset/Private/WorldConditionQueryConverter.cpp
@@ -0,0 +1,90 @@
+// Copyright Epic Games, Inc. All Rights Reserved.
+
+#include "WorldConditionQueryConverter.h"
+
+#include "Dom/JsonObject.h"
+#include "Dom/JsonValue.h"
+#include "WorldConditionQuery.h"
+
+namespace UE::WorldConditionsToolset
+{
+	FString FWorldConditionQueryConverter::GetName() const
+	{
+		return TEXT("WorldConditionQueryConverter");
+	}
+
+	bool FWorldConditionQueryConverter::CanConvertProperty(TNotNull<const FProperty*> Property)
+	{
+		const FStructProperty* const StructProp = CastField<FStructProperty>(static_cast<const FProperty*>(Property));
+		return StructProp != nullptr
+			&& StructProp->Struct == FWorldConditionQueryDefinition::StaticStruct();
+	}
+
+	TSharedPtr<FJsonObject> FWorldConditionQueryConverter::PropertyToJsonSchema(TNotNull<const FProperty*> Property)
+	{
+		TSharedPtr<FJsonObject> Schema = ToolsetStructToJsonSchema(FWorldConditionQueryDefinition::StaticStruct());
+		if (Schema.IsValid())
+		{
+			Schema->SetStringField(TEXT("title"), TEXT("WorldConditionQueryDefinition"));
+		}
+		return Schema;
+	}
+
+	TSharedPtr<FJsonValue> FWorldConditionQueryConverter::PropertyToDefault(
+		TNotNull<const FProperty*> Property, const FString& DefaultString)
+	{
+		return MakeShared<FJsonValueObject>(MakeShared<FJsonObject>());
+	}
+
+	TSharedPtr<FJsonValue> FWorldConditionQueryConverter::PropertyToJsonData(
+		TNotNull<FProperty*> Property, const void* Value)
+	{
+		const FWorldConditionQueryDefinition* const QueryDef =
+			static_cast<const FWorldConditionQueryDefinition*>(Value);
+		if (QueryDef == nullptr)
+		{
+			return MakeShared<FJsonValueNull>();
+		}
+
+		// Serialize all UPROPERTYs via the ToolsetJson pipeline.
+		// FInstancedStruct conditions get _structType automatically from JsonObjectConverter.
+		TSharedPtr<FJsonObject> Result = ToolsetStructToJsonData(
+			FWorldConditionQueryDefinition::StaticStruct(), QueryDef);
+		if (!Result.IsValid())
+		{
+			return MakeShared<FJsonValueNull>();
+		}
+
+		Result->SetStringField(TEXT("description"), QueryDef->GetDescription().ToString());
+		return MakeShared<FJsonValueObject>(Result);
+	}
+
+	bool FWorldConditionQueryConverter::JsonDataToProperty(
+		const TSharedPtr<FJsonValue>& JsonValue, TNotNull<FProperty*> Property, void* OutValue)
+	{
+		if (!JsonValue.IsValid() || JsonValue->Type != EJson::Object)
+		{
+			return false;
+		}
+
+		FWorldConditionQueryDefinition* const QueryDef =
+			static_cast<FWorldConditionQueryDefinition*>(OutValue);
+		if (QueryDef == nullptr)
+		{
+			return false;
+		}
+
+		// Deserialize all UPROPERTYs via the ToolsetJson pipeline.
+		// FInstancedStruct conditions are initialized from _structType automatically by JsonObjectConverter.
+		const TSharedPtr<FJsonObject> JsonObject = JsonValue->AsObject();
+		if (!ToolsetJsonDataToStruct(
+				JsonObject.ToSharedRef(), FWorldConditionQueryDefinition::StaticStruct(), QueryDef))
+		{
+			return false;
+		}
+
+		// Rebuild SharedDefinition from EditableConditions.
+		QueryDef->Initialize(nullptr);
+		return true;
+	}
+}
diff --git a/Engine/Plugins/Experimental/WorldConditionsToolset/Source/WorldConditionsToolset/Private/WorldConditionQueryConverter.h b/Engine/Plugins/Experimental/WorldConditionsToolset/Source/WorldConditionsToolset/Private/WorldConditionQueryConverter.h
new file mode 100644
--- /dev/null
+++ b/Engine/Plugins/Experimental/WorldConditionsToolset/Source/WorldConditionsToolset/Private/WorldConditionQueryConverter.h
@@ -0,0 +1,25 @@
+// Copyright Epic Games, Inc. All Rights Reserved.
+
+#pragma once
+
+#include "ToolsetRegistry/ToolsetJsonConverter.h"
+
+namespace UE::WorldConditionsToolset
+{
+	/**
+	 * JSON converter for FWorldConditionQueryDefinition.
+	 * Serializes the query's UPROPERTYs via the ToolsetJson pipeline,
+	 * plus a description convenience field.
+	 * On deserialization, rebuilds the SharedDefinition via Initialize().
+	 */
+	class FWorldConditionQueryConverter : public UE::ToolsetRegistry::FToolsetJsonConverter
+	{
+	public:
+		virtual FString GetName() const override;
+		virtual bool CanConvertProperty(TNotNull<const FProperty*> Property) override;
+		virtual TSharedPtr<FJsonObject> PropertyToJsonSchema(TNotNull<const FProperty*> Property) override;
+		virtual TSharedPtr<FJsonValue> PropertyToDefault(TNotNull<const FProperty*> Property, const FString& DefaultString) override;
+		virtual TSharedPtr<FJsonValue> PropertyToJsonData(TNotNull<FProperty*> Property, const void* Value) override;
+		virtual bool JsonDataToProperty(const TSharedPtr<FJsonValue>& JsonValue, TNotNull<FProperty*> Property, void* OutValue) override;
+	};
+}
diff --git a/Engine/Plugins/Experimental/WorldConditionsToolset/Source/WorldConditionsToolset/Private/WorldConditionTools.cpp b/Engine/Plugins/Experimental/WorldConditionsToolset/Source/WorldConditionsToolset/Private/WorldConditionTools.cpp
new file mode 100644
--- /dev/null
+++ b/Engine/Plugins/Experimental/WorldConditionsToolset/Source/WorldConditionsToolset/Private/WorldConditionTools.cpp
@@ -0,0 +1,21 @@
+// Copyright Epic Games, Inc. All Rights Reserved.
+
+#include "WorldConditionTools.h"
+
+#include "StructUtils/InstancedStruct.h"
+#include "WorldConditionBase.h"
+#include "WorldConditionQuery.h"
+
+FText UWorldConditionTools::GetQueryDescription(const FWorldConditionQueryDefinition& QueryDefinition)
+{
+	return QueryDefinition.GetDescription();
+}
+
+FText UWorldConditionTools::GetConditionDescription(const FInstancedStruct& Condition)
+{
+	if (const FWorldConditionBase* const WorldCondition = Condition.GetPtr<FWorldConditionBase>())
+	{
+		return WorldCondition->GetDescription();
+	}
+	return FText::GetEmpty();
+}
diff --git a/Engine/Plugins/Experimental/WorldConditionsToolset/Source/WorldConditionsToolset/Private/WorldConditionTools.h b/Engine/Plugins/Experimental/WorldConditionsToolset/Source/WorldConditionsToolset/Private/WorldConditionTools.h
new file mode 100644
--- /dev/null
+++ b/Engine/Plugins/Experimental/WorldConditionsToolset/Source/WorldConditionsToolset/Private/WorldConditionTools.h
@@ -0,0 +1,37 @@
+// Copyright Epic Games, Inc. All Rights Reserved.
+
+#pragma once
+
+#include "ToolsetRegistry/UToolsetRegistry.h"
+#include "WorldConditionTools.generated.h"
+
+struct FInstancedStruct;
+struct FWorldConditionQueryDefinition;
+
+/**
+ * Toolset for inspecting WorldCondition queries.
+ * Exposes human-readable descriptions of WorldCondition queries and individual conditions.
+ */
+UCLASS(Blueprintable)
+class UWorldConditionTools : public UToolsetDefinition
+{
+	GENERATED_BODY()
+
+public:
+	/**
+	 * Returns a human-readable description of a world condition query.
+	 * @param QueryDefinition The query definition to describe.
+	 * @return Text description of all conditions in the query.
+	 */
+	UFUNCTION(BlueprintCallable, Category = "WorldConditions")
+	static FText GetQueryDescription(const FWorldConditionQueryDefinition& QueryDefinition);
+
+	/**
+	 * Returns a human-readable description of a single world condition.
+	 * The condition must be passed as an FInstancedStruct containing an FWorldConditionBase-derived struct.
+	 * @param Condition The instanced struct holding the world condition.
+	 * @return Text description of the condition, or empty if invalid.
+	 */
+	UFUNCTION(BlueprintCallable, Category = "WorldConditions")
+	static FText GetConditionDescription(const FInstancedStruct& Condition);
+};
diff --git a/Engine/Plugins/Experimental/WorldConditionsToolset/Source/WorldConditionsToolset/Private/WorldConditionsToolset.cpp b/Engine/Plugins/Experimental/WorldConditionsToolset/Source/WorldConditionsToolset/Private/WorldConditionsToolset.cpp
new file mode 100644
--- /dev/null
+++ b/Engine/Plugins/Experimental/WorldConditionsToolset/Source/WorldConditionsToolset/Private/WorldConditionsToolset.cpp
@@ -0,0 +1,64 @@
+// Copyright Epic Games, Inc. All Rights Reserved.
+
+#include "WorldConditionsToolset.h"
+
+#include "Editor.h"
+#include "Modules/ModuleManager.h"
+#include "ToolsetRegistry/ToolsetRegistrySubsystem.h"
+#include "ToolsetRegistry/UToolsetRegistry.h"
+#include "WorldConditionQueryConverter.h"
+#include "WorldConditionTools.h"
+
+/*virtual*/ void FWorldConditionsToolsetModule::StartupModule()
+{
+	WorldConditionQueryConverter = MakeShared<UE::WorldConditionsToolset::FWorldConditionQueryConverter>();
+
+	if (GEditor != nullptr)
+	{
+		Register();
+	}
+	else
+	{
+		PostEngineInitHandle = FCoreDelegates::GetOnPostEngineInit().AddRaw(
+			this, &FWorldConditionsToolsetModule::Register);
+	}
+}
+
+/*virtual*/ void FWorldConditionsToolsetModule::ShutdownModule()
+{
+	FCoreDelegates::GetOnPostEngineInit().Remove(PostEngineInitHandle);
+
+	UToolsetRegistry::UnregisterToolsetClass(UWorldConditionTools::StaticClass());
+
+	if (GEditor != nullptr && WorldConditionQueryConverter.IsValid())
+	{
+		UToolsetRegistrySubsystem* const ToolsetRegistrySubsystem =
+			GEditor->GetEditorSubsystem<UToolsetRegistrySubsystem>();
+		if (ToolsetRegistrySubsystem != nullptr)
+		{
+			ToolsetRegistrySubsystem->ToolsetRegistry.UnregisterConverter(WorldConditionQueryConverter);
+		}
+	}
+
+	WorldConditionQueryConverter.Reset();
+}
+
+void FWorldConditionsToolsetModule::Register()
+{
+	if (GEditor == nullptr || !WorldConditionQueryConverter.IsValid())
+	{
+		return;
+	}
+
+	UToolsetRegistry::RegisterToolsetClass(UWorldConditionTools::StaticClass());
+
+	UToolsetRegistrySubsystem* const ToolsetRegistrySubsystem =
+		GEditor->GetEditorSubsystem<UToolsetRegistrySubsystem>();
+	if (ToolsetRegistrySubsystem != nullptr)
+	{
+		ToolsetRegistrySubsystem->ToolsetRegistry.RegisterConverter(WorldConditionQueryConverter);
+	}
+}
+
+IMPLEMENT_MODULE(FWorldConditionsToolsetModule, WorldConditionsToolset);
+DEFINE_LOG_CATEGORY(LogWorldConditionsToolset);
diff --git a/Engine/Plugins/Experimental/WorldConditionsToolset/Source/WorldConditionsToolset/Public/WorldConditionsToolset.h b/Engine/Plugins/Experimental/WorldConditionsToolset/Source/WorldConditionsToolset/Public/WorldConditionsToolset.h
new file mode 100644
--- /dev/null
+++ b/Engine/Plugins/Experimental/WorldConditionsToolset/Source/WorldConditionsToolset/Public/WorldConditionsToolset.h
@@ -0,0 +1,25 @@
+// Copyright Epic Games, Inc. All Rights Reserved.
+
+#pragma once
+
+#include "Modules/ModuleInterface.h"
+#include "Logging/LogMacros.h"
+#include "Delegates/IDelegateInstance.h"
+#include "Templates/SharedPointer.h"
+
+namespace UE::ToolsetRegistry { class FToolsetJsonConverter; }
+
+DECLARE_LOG_CATEGORY_EXTERN(LogWorldConditionsToolset, Log, Log);
+
+class FWorldConditionsToolsetModule : public IModuleInterface
+{
+public:
+	virtual void StartupModule() override;
+	virtual void ShutdownModule() override;
+
+private:
+	void Register();
+
+	TSharedPtr<UE::ToolsetRegistry::FToolsetJsonConverter> WorldConditionQueryConverter;
+	FDelegateHandle PostEngineInitHandle;
+};
diff --git a/Engine/Plugins/Experimental/WorldConditionsToolset/Source/WorldConditionsToolset/WorldConditionsToolset.Build.cs b/Engine/Plugins/Experimental/WorldConditionsToolset/Source/WorldConditionsToolset/WorldConditionsToolset.Build.cs
new file mode 100644
--- /dev/null
+++ b/Engine/Plugins/Experimental/WorldConditionsToolset/Source/WorldConditionsToolset/WorldConditionsToolset.Build.cs
@@ -0,0 +1,21 @@
+// Copyright Epic Games, Inc. All Rights Reserved.
+
+using UnrealBuildTool;
+
+public class WorldConditionsToolset : ModuleRules
+{
+	public WorldConditionsToolset(ReadOnlyTargetRules Target) : base(Target)
+	{
+		PCHUsage = PCHUsageMode.UseExplicitOrSharedPCHs;
+		PublicDependencyModuleNames.Add("Core");
+		PrivateDependencyModuleNames.AddRange(new string[]
+		{
+			"CoreUObject",
+			"Engine",
+			"Json",
+			"ToolsetRegistry",
+			"UnrealEd",
+			"WorldConditions",
+		});
+	}
+}
diff --git a/Engine/Plugins/Experimental/WorldConditionsToolset/WorldConditionsToolset.uplugin b/Engine/Plugins/Experimental/WorldConditionsToolset/WorldConditionsToolset.uplugin
new file mode 100644
--- /dev/null
+++ b/Engine/Plugins/Experimental/WorldConditionsToolset/WorldConditionsToolset.uplugin
@@ -0,0 +1,38 @@
+{
+	"FileVersion": 3,
+	"Version": 1,
+	"VersionName": "1.0",
+	"FriendlyName": "WorldConditionsToolset",
+	"Description": "Toolset for WorldConditions Inspection",
+	"Category": "Other",
+	"EditorOnly": true,
+	"CreatedBy": "Epic Games",
+	"CreatedByURL": "",
+	"DocsURL": "",
+	"MarketplaceURL": "",
+	"SupportURL": "",
+	"CanContainContent": false,
+	"IsBetaVersion": false,
+	"IsExperimentalVersion": true,
+	"NoRedist": true,
+	"Installed": false,
+	"EnabledByDefault": true,
+	"Modules": [
+		{
+			"Name": "WorldConditionsToolset",
+			"Type": "Editor",
+			"LoadingPhase": "Default",
+			"TargetAllowList": ["Editor"]
+		}
+	],
+	"Plugins": [
+		{
+			"Name": "WorldConditions",
+			"Enabled": true
+		},
+		{
+			"Name": "ToolsetRegistry",
+			"Enabled": true
+		}
+	]
+}
diff --git a/Engine/Plugins/Runtime/WorldConditions/Source/WorldConditions/Public/WorldConditionQuery.h b/Engine/Plugins/Runtime/WorldConditions/Source/WorldConditions/Public/WorldConditionQuery.h
--- a/Engine/Plugins/Runtime/WorldConditions/Source/WorldConditions/Public/WorldConditionQuery.h
+++ b/Engine/Plugins/Runtime/WorldConditions/Source/WorldConditions/Public/WorldConditionQuery.h
@@ -186,7 +186,7 @@ struct TStructOpsTypeTraits<FWorldConditionQuerySharedDefinition> : public TStru
  * The mutable state of the world condition is stored in FWorldConditionQueryState.
  * This allows to reuse the definitions and minimize the runtime memory needed to run queries.
  */
-USTRUCT()
+USTRUCT(BlueprintType)
 struct FWorldConditionQueryDefinition
 {
 	GENERATED_BODY()
