[Toolsets] Fix use-after-free in ShutdownModule during engine exit

During engine exit, PurgeAllUObjectsOnExit() runs before UnloadModulesAtShutdown(), so StaticClass() returns a dangling pointer when called from ShutdownModule(). Guard UObject-dependent cleanup with UObjectInitialized().

Changes:
- WorldConditionsToolset: guard UnregisterToolsetClass and GetEditorSubsystem calls

#Synced-CL 52127619

#Synced-CL 52127619

#Synced-CL 52127619

diff --git a/Engine/Plugins/Experimental/WorldConditionsToolset/Source/WorldConditionsToolset/Private/WorldConditionsToolset.cpp b/Engine/Plugins/Experimental/WorldConditionsToolset/Source/WorldConditionsToolset/Private/WorldConditionsToolset.cpp
--- a/Engine/Plugins/Experimental/WorldConditionsToolset/Source/WorldConditionsToolset/Private/WorldConditionsToolset.cpp
+++ b/Engine/Plugins/Experimental/WorldConditionsToolset/Source/WorldConditionsToolset/Private/WorldConditionsToolset.cpp
@@ -28,15 +28,18 @@
 {
 	FCoreDelegates::GetOnPostEngineInit().Remove(PostEngineInitHandle);
 
-	UToolsetRegistry::UnregisterToolsetClass(UWorldConditionTools::StaticClass());
-
-	if (GEditor != nullptr && WorldConditionQueryConverter.IsValid())
+	if (UObjectInitialized())
 	{
-		UToolsetRegistrySubsystem* const ToolsetRegistrySubsystem =
-			GEditor->GetEditorSubsystem<UToolsetRegistrySubsystem>();
-		if (ToolsetRegistrySubsystem != nullptr)
+		UToolsetRegistry::UnregisterToolsetClass(UWorldConditionTools::StaticClass());
+
+		if (GEditor != nullptr && WorldConditionQueryConverter.IsValid())
 		{
-			ToolsetRegistrySubsystem->ToolsetRegistry.UnregisterConverter(WorldConditionQueryConverter);
+			UToolsetRegistrySubsystem* const ToolsetRegistrySubsystem =
+				GEditor->GetEditorSubsystem<UToolsetRegistrySubsystem>();
+			if (ToolsetRegistrySubsystem != nullptr)
+			{
+				ToolsetRegistrySubsystem->ToolsetRegistry.UnregisterConverter(WorldConditionQueryConverter);
+			}
 		}
 	}
 
