Improve some AI toolsets descriptions.

#Synced-CL 52306751

diff --git a/Engine/Plugins/Experimental/WorldConditionsToolset/Source/WorldConditionsToolset/Private/WorldConditionTools.h b/Engine/Plugins/Experimental/WorldConditionsToolset/Source/WorldConditionsToolset/Private/WorldConditionTools.h
--- a/Engine/Plugins/Experimental/WorldConditionsToolset/Source/WorldConditionsToolset/Private/WorldConditionTools.h
+++ b/Engine/Plugins/Experimental/WorldConditionsToolset/Source/WorldConditionsToolset/Private/WorldConditionTools.h
@@ -8,10 +8,7 @@
 struct FInstancedStruct;
 struct FWorldConditionQueryDefinition;
 
-/**
- * Toolset for inspecting WorldCondition queries.
- * Exposes human-readable descriptions of WorldCondition queries and individual conditions.
- */
+/** Inspect WorldCondition (FWorldConditionQueryDefinition, FWorldConditionBase) structs. */
 UCLASS(Blueprintable)
 class UWorldConditionTools : public UToolsetDefinition
 {
