[AI Toolsets]: Move all toolsets to load at post engine init to simplify registration when toolset registry is available.

Also, clean up some targeting to ensure they're only ever packaged with the editor.

#Synced-CL 52411642

diff --git a/Engine/Plugins/Experimental/Toolsets/AIModuleToolset/AIModuleToolset.uplugin b/Engine/Plugins/Experimental/Toolsets/AIModuleToolset/AIModuleToolset.uplugin
--- a/Engine/Plugins/Experimental/Toolsets/AIModuleToolset/AIModuleToolset.uplugin
+++ b/Engine/Plugins/Experimental/Toolsets/AIModuleToolset/AIModuleToolset.uplugin
@@ -21,7 +21,7 @@
 		{
 			"Name": "AIModuleToolset",
 			"Type": "Editor",
-			"LoadingPhase": "Default",
+			"LoadingPhase": "PostEngineInit",
 			"TargetAllowList": ["Editor"]
 		}
 	],
diff --git a/Engine/Plugins/Experimental/Toolsets/AnimationAssistantToolset/AnimationAssistantToolset.uplugin b/Engine/Plugins/Experimental/Toolsets/AnimationAssistantToolset/AnimationAssistantToolset.uplugin
--- a/Engine/Plugins/Experimental/Toolsets/AnimationAssistantToolset/AnimationAssistantToolset.uplugin
+++ b/Engine/Plugins/Experimental/Toolsets/AnimationAssistantToolset/AnimationAssistantToolset.uplugin
@@ -17,8 +17,9 @@
 	"Modules": [
 		{
 			"Name": "AnimationAssistantToolset",
-			"Type": "Runtime",
-			"LoadingPhase": "Default"
+			"Type": "Editor",
+			"LoadingPhase": "PostEngineInit",
+			"TargetAllowList": ["Editor"]
 		}
 	],
 	"Plugins": [
diff --git a/Engine/Plugins/Experimental/Toolsets/ConversationToolset/ConversationToolset.uplugin b/Engine/Plugins/Experimental/Toolsets/ConversationToolset/ConversationToolset.uplugin
--- a/Engine/Plugins/Experimental/Toolsets/ConversationToolset/ConversationToolset.uplugin
+++ b/Engine/Plugins/Experimental/Toolsets/ConversationToolset/ConversationToolset.uplugin
@@ -21,7 +21,7 @@
 		{
 			"Name": "ConversationToolset",
 			"Type": "Editor",
-			"LoadingPhase": "Default",
+			"LoadingPhase": "PostEngineInit",
 			"TargetAllowList": ["Editor"]
 		}
 	],
diff --git a/Engine/Plugins/Experimental/Toolsets/MCPClientToolset/MCPClientToolset.uplugin b/Engine/Plugins/Experimental/Toolsets/MCPClientToolset/MCPClientToolset.uplugin
--- a/Engine/Plugins/Experimental/Toolsets/MCPClientToolset/MCPClientToolset.uplugin
+++ b/Engine/Plugins/Experimental/Toolsets/MCPClientToolset/MCPClientToolset.uplugin
@@ -25,7 +25,7 @@
 		{
 			"Name": "MCPClientToolset",
 			"Type": "Editor",
-			"LoadingPhase": "Default",
+			"LoadingPhase": "PostEngineInit",
 			"TargetAllowList": [ "Editor" ]
 		}
 	]
diff --git a/Engine/Plugins/Experimental/Toolsets/NiagaraAIAssistantTools/NiagaraAIAssistantTools.uplugin b/Engine/Plugins/Experimental/Toolsets/NiagaraAIAssistantTools/NiagaraAIAssistantTools.uplugin
--- a/Engine/Plugins/Experimental/Toolsets/NiagaraAIAssistantTools/NiagaraAIAssistantTools.uplugin
+++ b/Engine/Plugins/Experimental/Toolsets/NiagaraAIAssistantTools/NiagaraAIAssistantTools.uplugin
@@ -19,7 +19,7 @@
 		{
 			"Name": "NiagaraAIAssistantTools",
 			"Type": "EditorNoCommandlet",
-			"LoadingPhase": "Default",
+			"LoadingPhase": "PostEngineInit",
 			"TargetAllowList": ["Editor"]
 		}
 	],
diff --git a/Engine/Plugins/Experimental/Toolsets/SlateInspectorToolset/SlateInspectorToolset.uplugin b/Engine/Plugins/Experimental/Toolsets/SlateInspectorToolset/SlateInspectorToolset.uplugin
--- a/Engine/Plugins/Experimental/Toolsets/SlateInspectorToolset/SlateInspectorToolset.uplugin
+++ b/Engine/Plugins/Experimental/Toolsets/SlateInspectorToolset/SlateInspectorToolset.uplugin
@@ -19,12 +19,14 @@
 		{
 			"Name": "SlateInspectorToolset",
 			"Type": "Editor",
-			"LoadingPhase": "Default"
+			"LoadingPhase": "PostEngineInit",
+			"TargetAllowList": ["Editor"]
 		},
 		{
 			"Name": "SlateInspectorToolsetTests",
 			"Type": "Editor",
-			"LoadingPhase": "Default"
+			"LoadingPhase": "Default",
+			"TargetAllowList": ["Editor"]
 		}
 	],
 	"Plugins": [
diff --git a/Engine/Plugins/Experimental/Toolsets/StateTreeToolset/StateTreeToolset.uplugin b/Engine/Plugins/Experimental/Toolsets/StateTreeToolset/StateTreeToolset.uplugin
--- a/Engine/Plugins/Experimental/Toolsets/StateTreeToolset/StateTreeToolset.uplugin
+++ b/Engine/Plugins/Experimental/Toolsets/StateTreeToolset/StateTreeToolset.uplugin
@@ -21,7 +21,7 @@
 		{
 			"Name": "StateTreeToolset",
 			"Type": "Editor",
-			"LoadingPhase": "Default",
+			"LoadingPhase": "PostEngineInit",
 			"TargetAllowList": ["Editor"]
 		}
 	],
diff --git a/Engine/Plugins/Experimental/Toolsets/UMGToolSet/UMGToolSet.uplugin b/Engine/Plugins/Experimental/Toolsets/UMGToolSet/UMGToolSet.uplugin
--- a/Engine/Plugins/Experimental/Toolsets/UMGToolSet/UMGToolSet.uplugin
+++ b/Engine/Plugins/Experimental/Toolsets/UMGToolSet/UMGToolSet.uplugin
@@ -19,7 +19,7 @@
 		{
 			"Name": "UMGToolSet",
 			"Type": "EditorNoCommandlet",
-			"LoadingPhase": "Default",
+			"LoadingPhase": "PostEngineInit",
 			"TargetAllowList": ["Editor"]
 		}
 	],
diff --git a/Engine/Plugins/Experimental/Toolsets/WorldConditionsToolset/Source/WorldConditionsToolset/Private/WorldConditionsToolset.cpp b/Engine/Plugins/Experimental/Toolsets/WorldConditionsToolset/Source/WorldConditionsToolset/Private/WorldConditionsToolset.cpp
--- a/Engine/Plugins/Experimental/Toolsets/WorldConditionsToolset/Source/WorldConditionsToolset/Private/WorldConditionsToolset.cpp
+++ b/Engine/Plugins/Experimental/Toolsets/WorldConditionsToolset/Source/WorldConditionsToolset/Private/WorldConditionsToolset.cpp
@@ -12,22 +12,11 @@
 /*virtual*/ void FWorldConditionsToolsetModule::StartupModule()
 {
 	WorldConditionQueryConverter = MakeShared<UE::WorldConditionsToolset::FWorldConditionQueryConverter>();
-
-	if (GEditor != nullptr)
-	{
-		Register();
-	}
-	else
-	{
-		PostEngineInitHandle = FCoreDelegates::GetOnPostEngineInit().AddRaw(
-			this, &FWorldConditionsToolsetModule::Register);
-	}
+	Register();
 }
 
 /*virtual*/ void FWorldConditionsToolsetModule::ShutdownModule()
 {
-	FCoreDelegates::GetOnPostEngineInit().Remove(PostEngineInitHandle);
-
 	if (UObjectInitialized())
 	{
 		UToolsetRegistry::UnregisterToolsetClass(UWorldConditionTools::StaticClass());
diff --git a/Engine/Plugins/Experimental/Toolsets/WorldConditionsToolset/Source/WorldConditionsToolset/Public/WorldConditionsToolset.h b/Engine/Plugins/Experimental/Toolsets/WorldConditionsToolset/Source/WorldConditionsToolset/Public/WorldConditionsToolset.h
--- a/Engine/Plugins/Experimental/Toolsets/WorldConditionsToolset/Source/WorldConditionsToolset/Public/WorldConditionsToolset.h
+++ b/Engine/Plugins/Experimental/Toolsets/WorldConditionsToolset/Source/WorldConditionsToolset/Public/WorldConditionsToolset.h
@@ -4,7 +4,6 @@
 
 #include "Modules/ModuleInterface.h"
 #include "Logging/LogMacros.h"
-#include "Delegates/IDelegateInstance.h"
 #include "Templates/SharedPointer.h"
 
 namespace UE::ToolsetRegistry { class FToolsetJsonConverter; }
@@ -21,5 +20,4 @@ private:
 	void Register();
 
 	TSharedPtr<UE::ToolsetRegistry::FToolsetJsonConverter> WorldConditionQueryConverter;
-	FDelegateHandle PostEngineInitHandle;
 };
diff --git a/Engine/Plugins/Experimental/Toolsets/WorldConditionsToolset/WorldConditionsToolset.uplugin b/Engine/Plugins/Experimental/Toolsets/WorldConditionsToolset/WorldConditionsToolset.uplugin
--- a/Engine/Plugins/Experimental/Toolsets/WorldConditionsToolset/WorldConditionsToolset.uplugin
+++ b/Engine/Plugins/Experimental/Toolsets/WorldConditionsToolset/WorldConditionsToolset.uplugin
@@ -21,7 +21,7 @@
 		{
 			"Name": "WorldConditionsToolset",
 			"Type": "Editor",
-			"LoadingPhase": "Default",
+			"LoadingPhase": "PostEngineInit",
 			"TargetAllowList": ["Editor"]
 		}
 	],
